Removed PrivatePCH.h files that aren't needed.

diff --git a/Engine/Plugins/Animation/ControlRig/Source/ControlRig/Private/ControlRigPrivatePCH.h b/Engine/Plugins/Animation/ControlRig/Source/ControlRig/Private/ControlRigPrivatePCH.h
deleted file mode 100644
diff --git a/Engine/Plugins/Animation/ControlRig/Source/ControlRigEditor/Private/ControlRigEditorPrivatePCH.h b/Engine/Plugins/Animation/ControlRig/Source/ControlRigEditor/Private/ControlRigEditorPrivatePCH.h
deleted file mode 100644
diff --git a/Engine/Plugins/MovieScene/ActorSequence/Source/ActorSequence/Private/ActorSequencePrivatePCH.h b/Engine/Plugins/MovieScene/ActorSequence/Source/ActorSequence/Private/ActorSequencePrivatePCH.h
deleted file mode 100644
diff --git a/Engine/Plugins/Runtime/Analytics/Adjust/Source/AdjustEditor/Private/AdjustEditorPrivatePCH.h b/Engine/Plugins/Runtime/Analytics/Adjust/Source/AdjustEditor/Private/AdjustEditorPrivatePCH.h
deleted file mode 100644
diff --git a/Engine/Plugins/Runtime/Analytics/Adjust/Source/AndroidAdjust/Private/AndroidAdjust.cpp b/Engine/Plugins/Runtime/Analytics/Adjust/Source/AndroidAdjust/Private/AndroidAdjust.cpp
--- a/Engine/Plugins/Runtime/Analytics/Adjust/Source/AndroidAdjust/Private/AndroidAdjust.cpp
+++ b/Engine/Plugins/Runtime/Analytics/Adjust/Source/AndroidAdjust/Private/AndroidAdjust.cpp
@@ -2,7 +2,7 @@
 
 #include "AndroidAdjust.h"
 #include "Modules/ModuleManager.h"
-#include "AndroidAdjustPrivatePCH.h"
+#include "AndroidAdjustProvider.h"
 #include "Android/AndroidJNI.h"
 #include "Android/AndroidApplication.h"
 #include "Misc/Paths.h"
diff --git a/Engine/Plugins/Runtime/Analytics/Adjust/Source/AndroidAdjust/Private/AndroidAdjustPrivatePCH.h b/Engine/Plugins/Runtime/Analytics/Adjust/Source/AndroidAdjust/Private/AndroidAdjustPrivatePCH.h
deleted file mode 100644
diff --git a/Engine/Plugins/Runtime/Analytics/Adjust/Source/IOSAdjust/Private/IOSAdjust.cpp b/Engine/Plugins/Runtime/Analytics/Adjust/Source/IOSAdjust/Private/IOSAdjust.cpp
--- a/Engine/Plugins/Runtime/Analytics/Adjust/Source/IOSAdjust/Private/IOSAdjust.cpp
+++ b/Engine/Plugins/Runtime/Analytics/Adjust/Source/IOSAdjust/Private/IOSAdjust.cpp
@@ -1,7 +1,7 @@
 // Copyright Epic Games, Inc. All Rights Reserved.
 
 #include "IOSAdjust.h"
-#include "IOSAdjustPrivatePCH.h"
+#include "IOSAdjustProvider.h"
 #include "Misc/Paths.h"
 #include "Misc/ConfigCacheIni.h"
 
diff --git a/Engine/Plugins/Runtime/Analytics/Adjust/Source/IOSAdjust/Private/IOSAdjustPrivatePCH.h b/Engine/Plugins/Runtime/Analytics/Adjust/Source/IOSAdjust/Private/IOSAdjustPrivatePCH.h
deleted file mode 100644
diff --git a/Engine/Source/Editor/ClothingSystemEditor/Private/ClothingSystemEditorPrivatePCH.h b/Engine/Source/Editor/ClothingSystemEditor/Private/ClothingSystemEditorPrivatePCH.h
deleted file mode 100644
diff --git a/Engine/Source/Runtime/ClothingSystemRuntimeNv/Private/ClothingSystemRuntimePrivatePCH.h b/Engine/Source/Runtime/ClothingSystemRuntimeNv/Private/ClothingSystemRuntimePrivatePCH.h
deleted file mode 100644
